Fix XGE controller hanging due to botched xgconsole args.

diff --git a/Engine/Plugins/XGEController/Source/Private/XGEControllerModule.cpp b/Engine/Plugins/XGEController/Source/Private/XGEControllerModule.cpp
--- a/Engine/Plugins/XGEController/Source/Private/XGEControllerModule.cpp
+++ b/Engine/Plugins/XGEController/Source/Private/XGEControllerModule.cpp
@@ -352,7 +352,7 @@ void FXGEControllerModule::WriteOutThreadProc()
 		// and set the working directory of xgConsole.exe to the engine binaries folder below.
 		FString XGConsoleArgs = FString::Printf(TEXT("/VIRTUALIZEDIRECTX /allowremote=\"%s\" %s /allowintercept=\"%s\" /title=\"Unreal Engine XGE Tasks\" /monitordirs=\"%s\" /command=\"%s -xgecontroller %s\""),
 			XGE_INTERCEPT_EXE_NAMES,
-			GIsBuildMachine ? TEXT("") : TEXT("/avoidlocal=ON")
+			GIsBuildMachine ? TEXT("") : TEXT("/avoidlocal=ON"),
 			XGE_CONTROL_WORKER_NAME,
 			*WorkingDirectory,
 			XGE_CONTROL_WORKER_FILENAME,
