XGE controller: add a CVar and refine the logic.

- XGE tasks will be spawned on the local machine if any kind of commandlet is running.

#rnx

diff --git a/Engine/Plugins/XGEController/Source/Private/XGEControllerModule.cpp b/Engine/Plugins/XGEController/Source/Private/XGEControllerModule.cpp
--- a/Engine/Plugins/XGEController/Source/Private/XGEControllerModule.cpp
+++ b/Engine/Plugins/XGEController/Source/Private/XGEControllerModule.cpp
@@ -59,8 +59,36 @@ namespace XGEControllerVariables
 		Timeout,
 		TEXT("The time, in seconds, to wait after all tasks have been completed before shutting down the controller. (default: 2 seconds)."),
 		ECVF_Default);
+
+	int32 AvoidUsingLocalMachine = 1;
+	FAutoConsoleVariableRef CVarXGEControllerAvoidUsingLocalMachine(
+		TEXT("r.XGEController.AvoidUsingLocalMachine"),
+		AvoidUsingLocalMachine,
+		TEXT("Whether XGE tasks should avoid running on the local machine (to reduce the oversubscription with local async and out-of-process work).\n")
+		TEXT("0: Do not avoid. Distributed tasks will be spawned on all available XGE agents. Can cause oversubscription on the initiator machine. \n")
+		TEXT("1: Avoid spawning tasks on the local (initiator) machine except when running a commandlet or -buildmachine is passed (default).\n")
+		TEXT("2: Avoid spawning tasks on the local (initiator) machine."),
+		ECVF_Default); // This can be flipped any time XGEControlWorker is restarted
+
 }
 
+namespace XGEController
+{
+	bool AvoidUsingLocalMachine()
+	{
+		switch (XGEControllerVariables::AvoidUsingLocalMachine)
+		{
+			case 0:
+				return false;
+			case 2:
+				return true;
+			default:
+				break;
+		}
+
+		return !GIsBuildMachine && !IsRunningCommandlet();
+	}
+}
 
 FXGEControllerModule::FXGEControllerModule()
 	: bSupported(false)
@@ -352,7 +380,7 @@ void FXGEControllerModule::WriteOutThreadProc()
 		// and set the working directory of xgConsole.exe to the engine binaries folder below.
 		FString XGConsoleArgs = FString::Printf(TEXT("/VIRTUALIZEDIRECTX /allowremote=\"%s\" %s /allowintercept=\"%s\" /title=\"Unreal Engine XGE Tasks\" /monitordirs=\"%s\" /command=\"%s -xgecontroller %s\""),
 			XGE_INTERCEPT_EXE_NAMES,
-			GIsBuildMachine ? TEXT("") : TEXT("/avoidlocal=ON"),
+			XGEController::AvoidUsingLocalMachine() ? TEXT("/avoidlocal=ON") : TEXT(""),
 			XGE_CONTROL_WORKER_NAME,
 			*WorkingDirectory,
 			XGE_CONTROL_WORKER_FILENAME,
