Do not downprioritize XGE tasks, also print more info.

#rnx

diff --git a/Engine/Plugins/XGEController/Source/Private/XGEControllerModule.cpp b/Engine/Plugins/XGEController/Source/Private/XGEControllerModule.cpp
--- a/Engine/Plugins/XGEController/Source/Private/XGEControllerModule.cpp
+++ b/Engine/Plugins/XGEController/Source/Private/XGEControllerModule.cpp
@@ -74,6 +74,7 @@ namespace XGEControllerVariables
 
 namespace XGEController
 {
+	/** Whether XGE tasks should avoid using the local machine to reduce oversubscription */
 	bool AvoidUsingLocalMachine()
 	{
 		switch (XGEControllerVariables::AvoidUsingLocalMachine)
@@ -354,6 +355,11 @@ void FXGEControllerModule::InitializeController()
 		{
 			WriteOutThreadFuture = Async(EAsyncExecution::Thread, [this]() { WriteOutThreadProc(); });
 			bControllerInitialized = true;
+
+			UE_LOG(LogXGEController, Display, TEXT("Initialized XGE controller. XGE tasks will %sbe spawned on this machine."),
+				XGEController::AvoidUsingLocalMachine() ? TEXT("not ") : TEXT("")
+				);
+
 		}
 	}
 }
@@ -392,7 +398,7 @@ void FXGEControllerModule::WriteOutThreadProc()
 			UE_LOG(LogXGEController, Fatal, TEXT("Failed to create the output XGE named pipe."));
 		}
 
-		const int32 PriorityModifier = -1;	// below normal
+		const int32 PriorityModifier = 0;	// normal by default. Interactive use case shouldn't be affected as the jobs will avoid local machine
 		// Start the controller process
 		uint32 XGConsoleProcID = 0;
 		UE_LOG(LogXGEController, Verbose, TEXT("Launching xgConsole"));
diff --git a/Engine/Source/Programs/ShaderCompileWorker/Private/XGEControlWorker.cpp b/Engine/Source/Programs/ShaderCompileWorker/Private/XGEControlWorker.cpp
--- a/Engine/Source/Programs/ShaderCompileWorker/Private/XGEControlWorker.cpp
+++ b/Engine/Source/Programs/ShaderCompileWorker/Private/XGEControlWorker.cpp
@@ -173,8 +173,8 @@ void FXGEControlWorker::InputThreadProc()
 		Reader << Task->Executable;
 		Reader << Task->Arguments;
 
-		// Launch the process (with a lower than normal priority)
-		Task->Handle = FPlatformProcess::CreateProc(*Task->Executable, *Task->Arguments, true, false, false, nullptr, -1, nullptr, nullptr);
+		// Launch the process with normal priority.
+		Task->Handle = FPlatformProcess::CreateProc(*Task->Executable, *Task->Arguments, true, false, false, nullptr, 0, nullptr, nullptr);
 
 		FScopeLock Lock(CS);
 		CurrentTasks.Add(Task);
