Can now force using local machine on local XGE setups using the environment variable UE-XGEControllerForceUsingLocalMachine.

diff --git a/Engine/Plugins/XGEController/Source/Private/XGEControllerModule.cpp b/Engine/Plugins/XGEController/Source/Private/XGEControllerModule.cpp
--- a/Engine/Plugins/XGEController/Source/Private/XGEControllerModule.cpp
+++ b/Engine/Plugins/XGEController/Source/Private/XGEControllerModule.cpp
@@ -88,7 +88,8 @@ namespace XGEController
 				break;
 		}
 
-		return !GIsBuildMachine && !IsRunningCommandlet();
+		static bool bForceUsingLocalMachine = !FPlatformMisc::GetEnvironmentVariable(TEXT("UE-XGEControllerForceUsingLocalMachine")).IsEmpty();
+		return !bForceUsingLocalMachine && !GIsBuildMachine && !IsRunningCommandlet();
 	}
 }
 
