[shaders/xge] bypass the "ignore all throttling" guard from passing avoidlocal to xgconsole in the case we are running a commandlet. in commandlet execution IgnoreAllThrottling always returns true since commandlets run as part of the engine preinit step. with this we should actually be able to avoid local execution of XGE tasks on build machines.

diff --git a/Engine/Plugins/XGEController/Source/Private/XGEControllerModule.cpp b/Engine/Plugins/XGEController/Source/Private/XGEControllerModule.cpp
--- a/Engine/Plugins/XGEController/Source/Private/XGEControllerModule.cpp
+++ b/Engine/Plugins/XGEController/Source/Private/XGEControllerModule.cpp
@@ -383,7 +383,7 @@ void FXGEControllerModule::WriteOutThreadProc()
 		// and set the working directory of xgConsole.exe to the engine binaries folder below.
 		FString XGConsoleArgs = FString::Printf(TEXT("/VIRTUALIZEDIRECTX /allowremote=\"%s\" %s /allowintercept=\"%s\" /title=\"Unreal Engine XGE Tasks\" /monitordirs=\"%s\" /command=\"%s -xgecontroller %s\""),
 			XGE_INTERCEPT_EXE_NAMES,
-			(XGEController::AvoidUsingLocalMachine() && !GShaderCompilingManager->IgnoreAllThrottling()) ? TEXT("/avoidlocal=ON") : TEXT(""),
+			(XGEController::AvoidUsingLocalMachine() && (IsRunningCommandlet() || !GShaderCompilingManager->IgnoreAllThrottling())) ? TEXT("/avoidlocal=ON") : TEXT(""),
 			XGE_CONTROL_WORKER_NAME,
 			*WorkingDirectory,
 			XGE_CONTROL_WORKER_FILENAME,
