[XGE] Add XGE installation path to warning when it's installed but not actively running.

When XGEController plugin is enabled and XGE is installed but not actively running, we want to know where it is installed in case the warning is undesired so we can diagnose how to successully uninstall it.

#rnx

diff --git a/Engine/Plugins/XGEController/Source/Private/XGEControllerModule.cpp b/Engine/Plugins/XGEController/Source/Private/XGEControllerModule.cpp
--- a/Engine/Plugins/XGEController/Source/Private/XGEControllerModule.cpp
+++ b/Engine/Plugins/XGEController/Source/Private/XGEControllerModule.cpp
@@ -234,11 +234,11 @@ bool FXGEControllerModule::IsSupported()
 
 			if (Version == 0)
 			{
-				UE_LOG(LogXGEController, Warning, TEXT("Cannot determine XGE version. XGE Shader compilation with the interception interface may fail."));
+				UE_LOG(LogXGEController, Warning, TEXT("Cannot determine XGE version. XGE Shader compilation with the interception interface may fail [%s]."), *XGConsolePath);
 			}
 			else if (Version < 1002867)
 			{
-				UE_LOG(LogXGEController, Warning, TEXT("XGE version 8.01 (build 1867) or higher is required for XGE shader compilation with the interception interface."));
+				UE_LOG(LogXGEController, Warning, TEXT("XGE version 8.01 (build 1867) or higher is required for XGE shader compilation with the interception interface [%s]."), *XGConsolePath);
 				XGEControllerVariables::Enabled = 0;
 			}
 
@@ -246,13 +246,13 @@ bool FXGEControllerModule::IsSupported()
 			const TCHAR* XGEBuildServiceExecutableName = TEXT("BuildService.exe");
 			if (!FPlatformProcess::IsApplicationRunning(XGEBuildServiceExecutableName))
 			{
-				UE_LOG(LogXGEController, Warning, TEXT("XGE's background service (%s) is not running - service is likely disabled on this machine."), XGEBuildServiceExecutableName);
+				UE_LOG(LogXGEController, Warning, TEXT("XGE's background service (%s) is not running - service is likely disabled on this machine [%s]."), XGEBuildServiceExecutableName, *XGConsolePath);
 				XGEControllerVariables::Enabled = 0;
 			}
 
 			if (!PlatformFile.FileExists(*GetControlWorkerExePath()))
 			{
-				UE_LOG(LogXGEController, Warning, TEXT("XGEControlWorker.exe does not exist, XGE may be disabled in your Build Configuration, cannot use XGE."));
+				UE_LOG(LogXGEController, Warning, TEXT("XGEControlWorker.exe does not exist, XGE may be disabled in your Build Configuration, cannot use XGE [%s]."), *XGConsolePath);
 				XGEControllerVariables::Enabled = 0;
 			}
 		}
