[XGE] Lower verbosity in XGEControllerModule when service is not supported as UE is transitiong to UBA.

The IsSupported() function should only provide checks whether the module is supported or not but avoid spamming the build farm with warnings if it's not available due to certain conditions.
We are also moving away from XGE towards UBA. For those clients still relying on XGE, these warnings are now displayed with 'Display' verbosity so they still show up in the logs but not as a warning to improve CIS build health.
The affected warning messages were purly suggestive messages to begin with, such as XGE being installed but the service is not running.

#rnx

diff --git a/Engine/Plugins/XGEController/Source/Private/XGEControllerModule.cpp b/Engine/Plugins/XGEController/Source/Private/XGEControllerModule.cpp
--- a/Engine/Plugins/XGEController/Source/Private/XGEControllerModule.cpp
+++ b/Engine/Plugins/XGEController/Source/Private/XGEControllerModule.cpp
@@ -234,11 +234,11 @@ bool FXGEControllerModule::IsSupported()
 
 			if (Version == 0)
 			{
-				UE_LOG(LogXGEController, Warning, TEXT("Cannot determine XGE version. XGE Shader compilation with the interception interface may fail [%s]."), *XGConsolePath);
+				UE_LOG(LogXGEController, Display, TEXT("Cannot determine XGE version. XGE Shader compilation with the interception interface may fail [%s]."), *XGConsolePath);
 			}
 			else if (Version < 1002867)
 			{
-				UE_LOG(LogXGEController, Warning, TEXT("XGE version 8.01 (build 1867) or higher is required for XGE shader compilation with the interception interface [%s]."), *XGConsolePath);
+				UE_LOG(LogXGEController, Display, TEXT("XGE version 8.01 (build 1867) or higher is required for XGE shader compilation with the interception interface [%s]."), *XGConsolePath);
 				XGEControllerVariables::Enabled = 0;
 			}
 
@@ -246,13 +246,13 @@ bool FXGEControllerModule::IsSupported()
 			const TCHAR* XGEBuildServiceExecutableName = TEXT("BuildService.exe");
 			if (!FPlatformProcess::IsApplicationRunning(XGEBuildServiceExecutableName))
 			{
-				UE_LOG(LogXGEController, Warning, TEXT("XGE's background service (%s) is not running - service is likely disabled on this machine [%s]."), XGEBuildServiceExecutableName, *XGConsolePath);
+				UE_LOG(LogXGEController, Display, TEXT("XGE's background service (%s) is not running - service is likely disabled on this machine [%s]."), XGEBuildServiceExecutableName, *XGConsolePath);
 				XGEControllerVariables::Enabled = 0;
 			}
 
 			if (!PlatformFile.FileExists(*GetControlWorkerExePath()))
 			{
-				UE_LOG(LogXGEController, Warning, TEXT("XGEControlWorker.exe does not exist, XGE may be disabled in your Build Configuration, cannot use XGE [%s]."), *XGConsolePath);
+				UE_LOG(LogXGEController, Display, TEXT("XGEControlWorker.exe does not exist, XGE may be disabled in your Build Configuration, cannot use XGE [%s]."), *XGConsolePath);
 				XGEControllerVariables::Enabled = 0;
 			}
 		}
