[Shaders] Fix null pointer exception in OnDistributedShaderCompilingChanged().

#rnx

diff --git a/Engine/Plugins/XGEController/Source/Private/XGEControllerModule.cpp b/Engine/Plugins/XGEController/Source/Private/XGEControllerModule.cpp
--- a/Engine/Plugins/XGEController/Source/Private/XGEControllerModule.cpp
+++ b/Engine/Plugins/XGEController/Source/Private/XGEControllerModule.cpp
@@ -75,7 +75,10 @@ namespace XGEControllerVariables
 			[](IConsoleVariable* /*CVar*/) -> void
 			{
 				// If XGE has changed its support for local workers, notify shader compiling manager in case local shader compiling thread must be spun up.
-				GShaderCompilingManager->OnDistributedShaderCompilingChanged();
+				if (GShaderCompilingManager)
+				{
+					GShaderCompilingManager->OnDistributedShaderCompilingChanged();
+				}
 			}
 		),
 		ECVF_Default); // This can be flipped any time XGEControlWorker is restarted
diff --git a/Engine/Source/Runtime/Engine/Public/ShaderCompiler.h b/Engine/Source/Runtime/Engine/Public/ShaderCompiler.h
--- a/Engine/Source/Runtime/Engine/Public/ShaderCompiler.h
+++ b/Engine/Source/Runtime/Engine/Public/ShaderCompiler.h
@@ -976,7 +976,7 @@ private:
 	bool bIsEngineLoopInitialized;
 
 	/** Interface to the build distribution controller (XGE/SN-DBS) */
-	IDistributedBuildController* BuildDistributionController;
+	IDistributedBuildController* BuildDistributionController = nullptr;
 
 	/** Opt out of material shader compilation and instead place an empty shader map. */
 	bool bNoShaderCompilation;
